Allow script user to specify which LOD will be use to generate the physics asset. The intent is to provide a way to create a more clean physics asset because lower LODs normaly have some bone remove (dynamic and twist bones). Also when physics asset reference Bones, it prevent lower LOD from excluding those bones until the physics asset is cleanup.

#rnx

diff --git a/Engine/Source/Developer/MeshUtilities/Private/MeshUtilities.cpp b/Engine/Source/Developer/MeshUtilities/Private/MeshUtilities.cpp
--- a/Engine/Source/Developer/MeshUtilities/Private/MeshUtilities.cpp
+++ b/Engine/Source/Developer/MeshUtilities/Private/MeshUtilities.cpp
@@ -262,9 +262,9 @@ void FMeshUtilities::CalculateTriangleTangent(const FSoftSkinVertex& VertexA, co
 									 , OutTangents, CompareThreshold);
 }
 
-void FMeshUtilities::CalcBoneVertInfos(USkeletalMesh* SkeletalMesh, TArray<FBoneVertInfo>& Infos, bool bOnlyDominant)
+void FMeshUtilities::CalcBoneVertInfos(USkeletalMesh* SkeletalMesh, TArray<FBoneVertInfo>& Infos, bool bOnlyDominant, int32 LodIndex/* = 0*/)
 {
-	FMeshUtilitiesEngine::CalcBoneVertInfos(SkeletalMesh, Infos, bOnlyDominant);
+	FMeshUtilitiesEngine::CalcBoneVertInfos(SkeletalMesh, Infos, bOnlyDominant, LodIndex);
 }
 
 // Helper function for ConvertMeshesToStaticMesh
diff --git a/Engine/Source/Developer/MeshUtilities/Private/MeshUtilitiesPrivate.h b/Engine/Source/Developer/MeshUtilities/Private/MeshUtilitiesPrivate.h
--- a/Engine/Source/Developer/MeshUtilities/Private/MeshUtilitiesPrivate.h
+++ b/Engine/Source/Developer/MeshUtilities/Private/MeshUtilitiesPrivate.h
@@ -109,7 +109,7 @@ private:
 	 */
 	virtual void CalculateTriangleTangent(const FSoftSkinVertex& VertexA, const FSoftSkinVertex& VertexB, const FSoftSkinVertex& VertexC, TArray<FVector3f>& OutTangents, float CompareThreshold) override;
 
-	virtual void CalcBoneVertInfos(USkeletalMesh* SkeletalMesh, TArray<FBoneVertInfo>& Infos, bool bOnlyDominant) override;
+	virtual void CalcBoneVertInfos(USkeletalMesh* SkeletalMesh, TArray<FBoneVertInfo>& Infos, bool bOnlyDominant, int32 LodIndex = 0) override;
 
 	/** 
 	 * Convert a set of mesh components in their current pose to a static mesh. 
diff --git a/Engine/Source/Developer/MeshUtilities/Public/MeshUtilities.h b/Engine/Source/Developer/MeshUtilities/Public/MeshUtilities.h
--- a/Engine/Source/Developer/MeshUtilities/Public/MeshUtilities.h
+++ b/Engine/Source/Developer/MeshUtilities/Public/MeshUtilities.h
@@ -288,7 +288,7 @@ public:
 	 *	@param	Infos			The output array of vertices associated with each bone.
 	 *	@param	bOnlyDominant	Controls whether a vertex is added to the info for a bone if it is most controlled by that bone, or if that bone has ANY influence on that vert.
 	 */
-	virtual void CalcBoneVertInfos( USkeletalMesh* SkeletalMesh, TArray<FBoneVertInfo>& Infos, bool bOnlyDominant) = 0;
+	virtual void CalcBoneVertInfos( USkeletalMesh* SkeletalMesh, TArray<FBoneVertInfo>& Infos, bool bOnlyDominant, int32 LodIndex = 0) = 0;
 
 	/**
 	 * Convert a set of mesh components in their current pose to a static mesh. 
diff --git a/Engine/Source/Developer/MeshUtilitiesEngine/Private/MeshUtilitiesEngine.cpp b/Engine/Source/Developer/MeshUtilitiesEngine/Private/MeshUtilitiesEngine.cpp
--- a/Engine/Source/Developer/MeshUtilitiesEngine/Private/MeshUtilitiesEngine.cpp
+++ b/Engine/Source/Developer/MeshUtilitiesEngine/Private/MeshUtilitiesEngine.cpp
@@ -13,38 +13,40 @@ int32 GetDominantBoneIndex(FSoftSkinVertex* SoftVert)
 {
 	uint16 MaxWeightBone = 0;
 	uint16 MaxWeightWeight = 0;
-
-	for (int32 i = 0; i < MAX_TOTAL_INFLUENCES; i++)
+	for (int32 InfluenceIndex = 0; InfluenceIndex < MAX_TOTAL_INFLUENCES; InfluenceIndex++)
 	{
-		if (SoftVert->InfluenceWeights[i] > MaxWeightWeight)
+		if (SoftVert->InfluenceWeights[InfluenceIndex] > MaxWeightWeight)
 		{
-			MaxWeightWeight = SoftVert->InfluenceWeights[i];
-			MaxWeightBone = SoftVert->InfluenceBones[i];
+			MaxWeightWeight = SoftVert->InfluenceWeights[InfluenceIndex];
+			MaxWeightBone = SoftVert->InfluenceBones[InfluenceIndex];
 		}
 	}
 
 	return MaxWeightBone;
 }
 
-void FMeshUtilitiesEngine::CalcBoneVertInfos(USkeletalMesh* SkeletalMesh, TArray<FBoneVertInfo>& Infos, bool bOnlyDominant)
+void FMeshUtilitiesEngine::CalcBoneVertInfos(USkeletalMesh* SkeletalMesh, TArray<FBoneVertInfo>& Infos, bool bOnlyDominant, int32 SourceLodIndex)
 {
 	FSkeletalMeshModel* ImportedResource = SkeletalMesh->GetImportedModel();
 	if (ImportedResource->LODModels.Num() == 0)
 		return;
+	
+	//Find the source LOD index
+	const int32 LodIndex = SkeletalMesh->IsValidLODIndex(SourceLodIndex) ? SourceLodIndex : 0;
 
 	SkeletalMesh->CalculateInvRefMatrices();
 	check(SkeletalMesh->GetRefSkeleton().GetRawBoneNum() == SkeletalMesh->GetRefBasesInvMatrix().Num());
 
 	Infos.Empty();
 	Infos.AddZeroed(SkeletalMesh->GetRefSkeleton().GetRawBoneNum());
-
-	FSkeletalMeshLODModel* LODModel = &ImportedResource->LODModels[0];
+	
+	FSkeletalMeshLODModel* LODModel = &ImportedResource->LODModels[LodIndex];
 	for (int32 SectionIndex = 0; SectionIndex < LODModel->Sections.Num(); SectionIndex++)
 	{
 		FSkelMeshSection& Section = LODModel->Sections[SectionIndex];
-		for (int32 i = 0; i < Section.SoftVertices.Num(); i++)
+		for (int32 SectionVertexIndex = 0; SectionVertexIndex < Section.SoftVertices.Num(); SectionVertexIndex++)
 		{
-			FSoftSkinVertex* SoftVert = &Section.SoftVertices[i];
+			FSoftSkinVertex* SoftVert = &Section.SoftVertices[SectionVertexIndex];
 
 			if (bOnlyDominant)
 			{
@@ -58,12 +60,12 @@ void FMeshUtilitiesEngine::CalcBoneVertInfos(USkeletalMesh* SkeletalMesh, TArray
 			}
 			else
 			{
-				for (int32 j = 0; j < MAX_TOTAL_INFLUENCES; j++)
+				for (int32 InfluenceIndex = 0; InfluenceIndex < MAX_TOTAL_INFLUENCES; InfluenceIndex++)
 				{
-					if (SoftVert->InfluenceWeights[j] > 0)
+					if (SoftVert->InfluenceWeights[InfluenceIndex] > 0)
 					{
-						int32 BoneIndex = Section.BoneMap[SoftVert->InfluenceBones[j]];
-
+						int32 BoneIndex = Section.BoneMap[SoftVert->InfluenceBones[InfluenceIndex]];
+						
 						FVector3f LocalPos = SkeletalMesh->GetRefBasesInvMatrix()[BoneIndex].TransformPosition(SoftVert->Position);
 						Infos[BoneIndex].Positions.Add(LocalPos);
 
@@ -74,4 +76,4 @@ void FMeshUtilitiesEngine::CalcBoneVertInfos(USkeletalMesh* SkeletalMesh, TArray
 			}
 		}
 	}
-}
\ No newline at end of file
+}
diff --git a/Engine/Source/Developer/MeshUtilitiesEngine/Public/MeshUtilitiesEngine.h b/Engine/Source/Developer/MeshUtilitiesEngine/Public/MeshUtilitiesEngine.h
--- a/Engine/Source/Developer/MeshUtilitiesEngine/Public/MeshUtilitiesEngine.h
+++ b/Engine/Source/Developer/MeshUtilitiesEngine/Public/MeshUtilitiesEngine.h
@@ -17,6 +17,7 @@ public:
 	 *	@param	SkeletalMesh	The target skeletal mesh.
 	 *	@param	Infos			The output array of vertices associated with each bone.
 	 *	@param	bOnlyDominant	Controls whether a vertex is added to the info for a bone if it is most controlled by that bone, or if that bone has ANY influence on that vert.
+	 *	@param	SourceLodIndex	Use the specified LOD index to query the bone vert info.
 	 */
-	static void CalcBoneVertInfos(USkeletalMesh* SkeletalMesh, TArray<FBoneVertInfo>& Infos, bool bOnlyDominant);
+	static void CalcBoneVertInfos(USkeletalMesh* SkeletalMesh, TArray<FBoneVertInfo>& Infos, bool bOnlyDominant, int32 SourceLodIndex);
 };
\ No newline at end of file
diff --git a/Engine/Source/Developer/PhysicsUtilities/Private/PhysicsAssetUtils.cpp b/Engine/Source/Developer/PhysicsUtilities/Private/PhysicsAssetUtils.cpp
--- a/Engine/Source/Developer/PhysicsUtilities/Private/PhysicsAssetUtils.cpp
+++ b/Engine/Source/Developer/PhysicsUtilities/Private/PhysicsAssetUtils.cpp
@@ -295,7 +295,7 @@ bool CreateFromSkeletalMeshInternal(UPhysicsAsset* PhysicsAsset, USkeletalMesh*
 {
 	// For each bone, get the vertices most firmly attached to it.
 	TArray<FBoneVertInfo> Infos;
-	FMeshUtilitiesEngine::CalcBoneVertInfos(SkelMesh, Infos, (Params.VertWeight == EVW_DominantWeight));
+	FMeshUtilitiesEngine::CalcBoneVertInfos(SkelMesh, Infos, (Params.VertWeight == EVW_DominantWeight), Params.LodIndex);
 	check(Infos.Num() == SkelMesh->GetRefSkeleton().GetRawBoneNum());
 
 	PhysicsAsset->CollisionDisableTable.Empty();
diff --git a/Engine/Source/Developer/PhysicsUtilities/Public/PhysicsAssetUtils.h b/Engine/Source/Developer/PhysicsUtilities/Public/PhysicsAssetUtils.h
--- a/Engine/Source/Developer/PhysicsUtilities/Public/PhysicsAssetUtils.h
+++ b/Engine/Source/Developer/PhysicsUtilities/Public/PhysicsAssetUtils.h
@@ -57,6 +57,7 @@ struct FPhysAssetCreateParams
 		MaxHullVerts = 16;
 		LevelSetResolution = 8;
 		LatticeResolution = 8;
+		LodIndex = 0;
 	}
 
 	/** Bones that are shorter than this value will be ignored for body creation */
@@ -99,6 +100,10 @@ struct FPhysAssetCreateParams
 	UPROPERTY(EditAnywhere, Category = "Constraint Creation", meta=(EditCondition="bCreateConstraints"))
 	TEnumAsByte<EAngularConstraintMotion> AngularConstraintMode;
 
+	/** The LOD index use to create the physics asset.*/
+	UPROPERTY(EditAnywhere, Category = "Body Creation")
+	int32								LodIndex;
+
 	/** When creating multiple convex hulls, the maximum number that will be created. */
 	UPROPERTY(EditAnywhere, AdvancedDisplay, Category = "Body Creation", meta = (EditCondition = "GeomType == EPhysAssetFitGeomType::EFG_MultiConvexHull", EditConditionHides))
 	int32								HullCount;
diff --git a/Engine/Source/Editor/PhysicsAssetEditor/Private/PhysicsAssetEditorSharedData.cpp b/Engine/Source/Editor/PhysicsAssetEditor/Private/PhysicsAssetEditorSharedData.cpp
--- a/Engine/Source/Editor/PhysicsAssetEditor/Private/PhysicsAssetEditorSharedData.cpp
+++ b/Engine/Source/Editor/PhysicsAssetEditor/Private/PhysicsAssetEditorSharedData.cpp
@@ -3563,7 +3563,10 @@ TSharedRef<SWidget> FPhysicsAssetEditorSharedData::CreateGenerateBodiesWidget(co
 	FPropertyEditorModule& PropertyEditorModule = FModuleManager::LoadModuleChecked<FPropertyEditorModule>("PropertyEditor");
 	TSharedRef<IDetailsView> DetailsView = PropertyEditorModule.CreateDetailView(DetailsViewArgs);
 
-	GetMutableDefault<UPhysicsAssetGenerationSettings>()->LoadConfig();
+	UPhysicsAssetGenerationSettings* PhysicsAssetGenerationSettings = GetMutableDefault<UPhysicsAssetGenerationSettings>();
+	PhysicsAssetGenerationSettings->LoadConfig();
+	// LodIndex value must be reset to zero when opening the dialog
+	PhysicsAssetGenerationSettings->CreateParams.LodIndex = 0;
 	DetailsView->SetObject(GetMutableDefault<UPhysicsAssetGenerationSettings>());
 	DetailsView->OnFinishedChangingProperties().AddLambda([](const FPropertyChangedEvent& InEvent){ GetMutableDefault<UPhysicsAssetGenerationSettings>()->SaveConfig(); });
 
diff --git a/Engine/Source/Editor/SkeletalMeshEditor/Private/SkeletalMeshEditorSubsystem.cpp b/Engine/Source/Editor/SkeletalMeshEditor/Private/SkeletalMeshEditorSubsystem.cpp
--- a/Engine/Source/Editor/SkeletalMeshEditor/Private/SkeletalMeshEditorSubsystem.cpp
+++ b/Engine/Source/Editor/SkeletalMeshEditor/Private/SkeletalMeshEditorSubsystem.cpp
@@ -932,7 +932,7 @@ void USkeletalMeshEditorSubsystem::SetLodBuildSettings(USkeletalMesh* SkeletalMe
 	}
 }
 
-UPhysicsAsset* USkeletalMeshEditorSubsystem::CreatePhysicsAsset(USkeletalMesh* SkeletalMesh)
+UPhysicsAsset* USkeletalMeshEditorSubsystem::CreatePhysicsAsset(USkeletalMesh* SkeletalMesh, bool bSetToMesh /*= true*/, int32 LodIndex /*= 0*/)
 {
 	if (!SkeletalMesh)
 	{
@@ -969,9 +969,9 @@ UPhysicsAsset* USkeletalMeshEditorSubsystem::CreatePhysicsAsset(USkeletalMesh* S
 		NewPhysicsAsset->MarkPackageDirty();
 
 		FPhysAssetCreateParams NewBodyData;
+		NewBodyData.LodIndex = LodIndex;
 		FText CreationErrorMessage;
-
-		bool bSuccess = FPhysicsAssetUtils::CreateFromSkeletalMesh(NewPhysicsAsset, SkeletalMesh, NewBodyData, CreationErrorMessage);
+		bool bSuccess = FPhysicsAssetUtils::CreateFromSkeletalMesh(NewPhysicsAsset, SkeletalMesh, NewBodyData, CreationErrorMessage, bSetToMesh);
 		if (!bSuccess)
 		{
 			UE_LOG(LogSkeletalMeshEditorSubsystem, Error, TEXT("CreatePhysicsAsset failed: Couldn't create PhysicsAsset for the SkeletalMesh."));
diff --git a/Engine/Source/Editor/SkeletalMeshEditor/Public/SkeletalMeshEditorSubsystem.h b/Engine/Source/Editor/SkeletalMeshEditor/Public/SkeletalMeshEditorSubsystem.h
--- a/Engine/Source/Editor/SkeletalMeshEditor/Public/SkeletalMeshEditorSubsystem.h
+++ b/Engine/Source/Editor/SkeletalMeshEditor/Public/SkeletalMeshEditorSubsystem.h
@@ -292,9 +292,11 @@ public:
 	 * This function creates a PhysicsAsset for the given SkeletalMesh with the same settings as if it were created through FBX import
 	 *
 	 * @Param SkeletalMesh: The SkeletalMesh we want to create the PhysicsAsset for
+	 * @Param bSetToMesh: If true the created physics asset will be used by the skeletal mesh.
+	 * @Param LodIndex: Use the specified LOD to create the physics asset.
 	 */
 	UFUNCTION(BlueprintCallable, Category = "Editor Scripting | SkeletalMesh")
-	static UPhysicsAsset* CreatePhysicsAsset(USkeletalMesh* SkeletalMesh);
+	static UPhysicsAsset* CreatePhysicsAsset(USkeletalMesh* SkeletalMesh, bool bSetToMesh = true, int32 LodIndex = 0);
 
 	/**
 	 * Checks whether a physics asset is compatible with the given SkeletalMesh
diff --git a/Engine/Source/Editor/UnrealEd/Private/Factories/PhysicsAssetFactory.cpp b/Engine/Source/Editor/UnrealEd/Private/Factories/PhysicsAssetFactory.cpp
--- a/Engine/Source/Editor/UnrealEd/Private/Factories/PhysicsAssetFactory.cpp
+++ b/Engine/Source/Editor/UnrealEd/Private/Factories/PhysicsAssetFactory.cpp
@@ -129,6 +129,13 @@ UObject* UPhysicsAssetFactory::CreatePhysicsAssetFromMesh(FName InAssetName, UOb
 			FText ErrorMessage;
 			const FPhysAssetCreateParams& NewBodyData = GetDefault<UPhysicsAssetGenerationSettings>()->CreateParams;
 			bool bSuccess = FPhysicsAssetUtils::CreateFromSkeletalMesh(NewAsset, SkelMesh, NewBodyData, ErrorMessage, bSetToMesh);
+			//Reset any property that should not be reconduct next time the dialog is open
+			//	- the LodIndex must be reset to zero
+			if (UPhysicsAssetGenerationSettings* MutablePhysicsAssetGenerationSettings = GetMutableDefault<UPhysicsAssetGenerationSettings>())
+			{
+				MutablePhysicsAssetGenerationSettings->CreateParams.LodIndex = 0;
+				MutablePhysicsAssetGenerationSettings->SaveConfig();
+			}
 			if(bSuccess)
 			{
 				NewAsset->MarkPackageDirty();
