Physics: Allow computing the per-bone vertex positions on cooked builds.

#rnx

diff --git a/Engine/Source/Developer/MeshUtilitiesEngine/Private/MeshUtilitiesEngine.cpp b/Engine/Source/Developer/MeshUtilitiesEngine/Private/MeshUtilitiesEngine.cpp
--- a/Engine/Source/Developer/MeshUtilitiesEngine/Private/MeshUtilitiesEngine.cpp
+++ b/Engine/Source/Developer/MeshUtilitiesEngine/Private/MeshUtilitiesEngine.cpp
@@ -2,23 +2,22 @@
 
 #include "MeshUtilitiesEngine.h"
 
-#include "Engine/SkeletalMesh.h"
-#include "GPUSkinPublicDefs.h"
 #include "MeshUtilitiesCommon.h"
-#include "Rendering/SkeletalMeshLODModel.h"
-#include "Rendering/SkeletalMeshModel.h"
+#include "Engine/SkeletalMesh.h"
+#include "Rendering/SkeletalMeshRenderData.h"
+#include "Rendering/SkinWeightVertexBuffer.h"
 
 // Find the most dominant bone for each vertex
-int32 GetDominantBoneIndex(FSoftSkinVertex* SoftVert)
+FBoneIndexType GetDominantBoneIndex(const FSkinWeightInfo& VertexInfluences)
 {
-	uint16 MaxWeightBone = 0;
+	FBoneIndexType MaxWeightBone = 0;
 	uint16 MaxWeightWeight = 0;
 	for (int32 InfluenceIndex = 0; InfluenceIndex < MAX_TOTAL_INFLUENCES; InfluenceIndex++)
 	{
-		if (SoftVert->InfluenceWeights[InfluenceIndex] > MaxWeightWeight)
+		if (VertexInfluences.InfluenceWeights[InfluenceIndex] > MaxWeightWeight)
 		{
-			MaxWeightWeight = SoftVert->InfluenceWeights[InfluenceIndex];
-			MaxWeightBone = SoftVert->InfluenceBones[InfluenceIndex];
+			MaxWeightWeight = VertexInfluences.InfluenceWeights[InfluenceIndex];
+			MaxWeightBone = VertexInfluences.InfluenceBones[InfluenceIndex];
 		}
 	}
 
@@ -27,10 +26,6 @@ int32 GetDominantBoneIndex(FSoftSkinVertex* SoftVert)
 
 void FMeshUtilitiesEngine::CalcBoneVertInfos(USkeletalMesh* SkeletalMesh, TArray<FBoneVertInfo>& Infos, bool bOnlyDominant, int32 SourceLodIndex)
 {
-	FSkeletalMeshModel* ImportedResource = SkeletalMesh->GetImportedModel();
-	if (ImportedResource->LODModels.Num() == 0)
-		return;
-	
 	//Find the source LOD index
 	const int32 LodIndex = SkeletalMesh->IsValidLODIndex(SourceLodIndex) ? SourceLodIndex : 0;
 
@@ -38,38 +33,58 @@ void FMeshUtilitiesEngine::CalcBoneVertInfos(USkeletalMesh* SkeletalMesh, TArray
 	check(SkeletalMesh->GetRefSkeleton().GetRawBoneNum() == SkeletalMesh->GetRefBasesInvMatrix().Num());
 
 	Infos.Empty();
-	Infos.AddZeroed(SkeletalMesh->GetRefSkeleton().GetRawBoneNum());
 	
-	FSkeletalMeshLODModel* LODModel = &ImportedResource->LODModels[LodIndex];
-	for (int32 SectionIndex = 0; SectionIndex < LODModel->Sections.Num(); SectionIndex++)
+	const FSkeletalMeshRenderData* RenderData = SkeletalMesh->GetResourceForRendering();
+	if (!RenderData || !RenderData->LODRenderData.IsValidIndex(LodIndex))
 	{
-		FSkelMeshSection& Section = LODModel->Sections[SectionIndex];
-		for (int32 SectionVertexIndex = 0; SectionVertexIndex < Section.SoftVertices.Num(); SectionVertexIndex++)
-		{
-			FSoftSkinVertex* SoftVert = &Section.SoftVertices[SectionVertexIndex];
+		return;
+	}
+
+	const FSkeletalMeshLODRenderData& LODRenderData = RenderData->LODRenderData[LodIndex];
+	if (!LODRenderData.StaticVertexBuffers.PositionVertexBuffer.GetAllowCPUAccess() ||
+		!LODRenderData.StaticVertexBuffers.StaticMeshVertexBuffer.GetAllowCPUAccess() ||
+		!LODRenderData.SkinWeightVertexBuffer.GetNeedsCPUAccess())
+	{
+		return;
+	}
 
+	const FPositionVertexBuffer& PositionBuffer = LODRenderData.StaticVertexBuffers.PositionVertexBuffer;
+	const FStaticMeshVertexBuffer& TangentBuffer = LODRenderData.StaticVertexBuffers.StaticMeshVertexBuffer;
+
+	Infos.AddZeroed(SkeletalMesh->GetRefSkeleton().GetRawBoneNum());
+
+	TArray<FSkinWeightInfo> SkinWeights;
+	LODRenderData.SkinWeightVertexBuffer.GetSkinWeights(SkinWeights);
+
+	for (int32 SectionIndex = 0; SectionIndex < LODRenderData.RenderSections.Num(); SectionIndex++)
+	{
+		const FSkelMeshRenderSection& RenderSection =  LODRenderData.RenderSections[SectionIndex];
+		for (int32 SectionVertexIndex = 0; SectionVertexIndex < static_cast<int32>(RenderSection.NumVertices); SectionVertexIndex++)
+		{
+			const int32 VertexIndex = SectionVertexIndex + RenderSection.BaseVertexIndex;
 			if (bOnlyDominant)
 			{
-				int32 BoneIndex = Section.BoneMap[GetDominantBoneIndex(SoftVert)];
+				const int32 BoneIndex = RenderSection.BoneMap[GetDominantBoneIndex(SkinWeights[VertexIndex])];
 
-				FVector3f LocalPos = SkeletalMesh->GetRefBasesInvMatrix()[BoneIndex].TransformPosition(SoftVert->Position);
+				FVector3f LocalPos = SkeletalMesh->GetRefBasesInvMatrix()[BoneIndex].TransformPosition(PositionBuffer.VertexPosition(VertexIndex));
 				Infos[BoneIndex].Positions.Add(LocalPos);
 
-				FVector3f LocalNormal = SkeletalMesh->GetRefBasesInvMatrix()[BoneIndex].TransformVector(SoftVert->TangentZ);
+				FVector3f LocalNormal = SkeletalMesh->GetRefBasesInvMatrix()[BoneIndex].TransformVector(TangentBuffer.VertexTangentZ(VertexIndex));
 				Infos[BoneIndex].Normals.Add(LocalNormal);
 			}
 			else
 			{
+				const FSkinWeightInfo& BoneWeightInfo = SkinWeights[VertexIndex];
 				for (int32 InfluenceIndex = 0; InfluenceIndex < MAX_TOTAL_INFLUENCES; InfluenceIndex++)
 				{
-					if (SoftVert->InfluenceWeights[InfluenceIndex] > 0)
+					if (BoneWeightInfo.InfluenceWeights[InfluenceIndex] > 0)
 					{
-						int32 BoneIndex = Section.BoneMap[SoftVert->InfluenceBones[InfluenceIndex]];
-						
-						FVector3f LocalPos = SkeletalMesh->GetRefBasesInvMatrix()[BoneIndex].TransformPosition(SoftVert->Position);
+						int32 BoneIndex = RenderSection.BoneMap[BoneWeightInfo.InfluenceBones[InfluenceIndex]];
+					
+						FVector3f LocalPos = SkeletalMesh->GetRefBasesInvMatrix()[BoneIndex].TransformPosition(PositionBuffer.VertexPosition(VertexIndex));
 						Infos[BoneIndex].Positions.Add(LocalPos);
 
-						FVector3f LocalNormal = SkeletalMesh->GetRefBasesInvMatrix()[BoneIndex].TransformVector(SoftVert->TangentZ);
+						FVector3f LocalNormal = SkeletalMesh->GetRefBasesInvMatrix()[BoneIndex].TransformVector(TangentBuffer.VertexTangentZ(VertexIndex));
 						Infos[BoneIndex].Normals.Add(LocalNormal);
 					}
 				}
